[Mesh Partition]

Fixed MeshTerrainMode ui style referencing the location of the plugin through the absolute path rather than by querying the plugin directory.

#Synced-CL 51480380

diff --git a/Engine/Plugins/Experimental/MeshTerrainMode/Source/MeshTerrainMode/Private/MeshTerrainModeStyle.cpp b/Engine/Plugins/Experimental/MeshTerrainMode/Source/MeshTerrainMode/Private/MeshTerrainModeStyle.cpp
--- a/Engine/Plugins/Experimental/MeshTerrainMode/Source/MeshTerrainMode/Private/MeshTerrainModeStyle.cpp
+++ b/Engine/Plugins/Experimental/MeshTerrainMode/Source/MeshTerrainMode/Private/MeshTerrainModeStyle.cpp
@@ -66,7 +66,8 @@ void FMeshTerrainModeStyle::Initialize()
 	// If we get asked for something that we don't set, we should default to editor style
 	StyleSet->SetParentStyleName("EditorStyle");
 
-	StyleSet->SetContentRoot(FPaths::EngineDir() / TEXT("Restricted/NotForLicensees/Plugins/MeshTerrainMode/Content"));
+	const FString PluginContentDir = FPaths::Combine(IPluginManager::Get().FindPlugin(TEXT("MeshTerrainMode"))->GetBaseDir(), TEXT("Content"));
+	StyleSet->SetContentRoot(PluginContentDir);
 	StyleSet->SetCoreContentRoot(FPaths::EngineContentDir() / TEXT("Slate"));
 	
 	const FTextBlockStyle& NormalText = FAppStyle::Get().GetWidgetStyle<FTextBlockStyle>("NormalText");
