opt mesh mode out of RequiresLegacyViewportInteractions again

#Synced-CL 51646291

diff --git a/Engine/Plugins/Experimental/MeshTerrainMode/Source/MeshTerrainMode/Private/MeshTerrainMode.h b/Engine/Plugins/Experimental/MeshTerrainMode/Source/MeshTerrainMode/Private/MeshTerrainMode.h
--- a/Engine/Plugins/Experimental/MeshTerrainMode/Source/MeshTerrainMode/Private/MeshTerrainMode.h
+++ b/Engine/Plugins/Experimental/MeshTerrainMode/Source/MeshTerrainMode/Private/MeshTerrainMode.h
@@ -63,6 +63,8 @@ public:
 
 	virtual bool GetPivotForOrbit(FVector& OutPivot) const override;
 
+	virtual bool RequiresLegacyViewportInteractions() const override { return false; }
+
 	/*
 	 * focus events
 	 */
