[Mesh Partition]

Added plugin descriptions and documentation links.

#Synced-CL 51841952

diff --git a/Engine/Plugins/Experimental/MeshPartition/MeshPartition.uplugin b/Engine/Plugins/Experimental/MeshPartition/MeshPartition.uplugin
--- a/Engine/Plugins/Experimental/MeshPartition/MeshPartition.uplugin
+++ b/Engine/Plugins/Experimental/MeshPartition/MeshPartition.uplugin
@@ -2,12 +2,12 @@
 	"FileVersion": 3,
 	"Version": 1,
 	"VersionName": "1.0",
-	"FriendlyName": "MeshPartition",
-	"Description": "",
-	"Category": "Other",
-	"CreatedBy": "",
-	"CreatedByURL": "",
-	"DocsURL": "",
+	"FriendlyName": "Mesh Partition",
+	"Description": "Large-scale mesh authoring system through spatial partitioning, non-destructive modifier editing, and platform-adaptive runtime representations.",
+	"Category": "Mesh Partition",
+	"CreatedBy": "Epic Games, Inc.",
+	"CreatedByURL": "https://epicgames.com",
+	"DocsURL": "https://dev.epicgames.com/community/learning/knowledge-base/nK7J/unreal-engine-introduction-to-mesh-terrain",
 	"MarketplaceURL": "",
 	"SupportURL": "",
 	"CanContainContent": true,
@@ -40,7 +40,6 @@
 			"Type": "Editor",
 			"LoadingPhase": "Default"
 		}
-
 	],
 	"Plugins": [
 		{
@@ -58,17 +57,23 @@
 		{
 			"Name": "MeshModelingToolset",
 			"Enabled": true,
-			"TargetAllowList": [ "Editor" ]
+			"TargetAllowList": [
+				"Editor"
+			]
 		},
 		{
 			"Name": "ModelingToolsEditorMode",
 			"Enabled": true,
-			"TargetAllowList": [ "Editor" ]
+			"TargetAllowList": [
+				"Editor"
+			]
 		},
 		{
 			"Name": "EditorDataStorageFeatures",
 			"Enabled": true,
-			"TargetAllowList": [ "Editor" ]
+			"TargetAllowList": [
+				"Editor"
+			]
 		}
 	]
 }
\ No newline at end of file
diff --git a/Engine/Plugins/Experimental/MeshPartitionWater/MeshPartitionWater.uplugin b/Engine/Plugins/Experimental/MeshPartitionWater/MeshPartitionWater.uplugin
--- a/Engine/Plugins/Experimental/MeshPartitionWater/MeshPartitionWater.uplugin
+++ b/Engine/Plugins/Experimental/MeshPartitionWater/MeshPartitionWater.uplugin
@@ -2,12 +2,12 @@
 	"FileVersion": 3,
 	"Version": 1,
 	"VersionName": "1.0",
-	"FriendlyName": "MeshPartitionWater",
-	"Description": "",
-	"Category": "Other",
-	"CreatedBy": "",
-	"CreatedByURL": "",
-	"DocsURL": "",
+	"FriendlyName": "Mesh Partition Water",
+	"Description": "Interoperability of Mesh Partition with the Water plugin.",
+	"Category": "Mesh Partition",
+	"CreatedBy": "Epic Games, Inc.",
+	"CreatedByURL": "https://epicgames.com",
+	"DocsURL": "https://dev.epicgames.com/community/learning/knowledge-base/nK7J/unreal-engine-introduction-to-mesh-terrain#usingwaterbodytools",
 	"MarketplaceURL": "",
 	"SupportURL": "",
 	"CanContainContent": true,
diff --git a/Engine/Plugins/Experimental/MeshTerrainMode/MeshTerrainMode.uplugin b/Engine/Plugins/Experimental/MeshTerrainMode/MeshTerrainMode.uplugin
--- a/Engine/Plugins/Experimental/MeshTerrainMode/MeshTerrainMode.uplugin
+++ b/Engine/Plugins/Experimental/MeshTerrainMode/MeshTerrainMode.uplugin
@@ -1,13 +1,13 @@
 {
 	"FileVersion": 3,
 	"Version": 1,
-	"VersionName": "0.1",
+	"VersionName": "1.0",
 	"FriendlyName": "Mesh Terrain Mode",
-	"Description": "Mesh Terrain Mode includes a suite of interactive tools for creating and editing Mesh Terrain in the Editor",
-	"Category": "Editor",
+	"Description": "Mesh Terrain Mode includes a suite of interactive tools for creating and editing Mesh Partitions in the Editor",
+	"Category": "Mesh Partition",
 	"CreatedBy": "Epic Games, Inc.",
 	"CreatedByURL": "https://epicgames.com",
-	"DocsURL": "",
+	"DocsURL": "https://dev.epicgames.com/community/learning/knowledge-base/nK7J/unreal-engine-introduction-to-mesh-terrain#meshterrainmode",
 	"MarketplaceURL": "",
 	"SupportURL": "",
 	"CanContainContent": true,
@@ -26,37 +26,51 @@
 		{
 			"Name": "MeshModelingToolset",
 			"Enabled": true,
-			"TargetAllowList": [ "Editor" ]
+			"TargetAllowList": [
+				"Editor"
+			]
 		},
 		{
 			"Name": "MeshModelingToolsetExp",
 			"Enabled": true,
-			"TargetAllowList": [ "Editor" ]
+			"TargetAllowList": [
+				"Editor"
+			]
 		},
 		{
 			"Name": "MeshLODToolset",
 			"Enabled": true,
-			"TargetAllowList": [ "Editor" ]
+			"TargetAllowList": [
+				"Editor"
+			]
 		},
 		{
 			"Name": "ToolPresets",
 			"Enabled": true,
-			"TargetAllowList": [ "Editor" ]
+			"TargetAllowList": [
+				"Editor"
+			]
 		},
 		{
 			"Name": "StylusInput",
 			"Enabled": true,
-			"TargetAllowList": [ "Editor" ]
+			"TargetAllowList": [
+				"Editor"
+			]
 		},
 		{
 			"Name": "ModelingToolsEditorMode",
 			"Enabled": true,
-			"TargetAllowList": [ "Editor" ]
+			"TargetAllowList": [
+				"Editor"
+			]
 		},
 		{
 			"Name": "MeshPartition",
 			"Enabled": true,
-			"TargetAllowList": [ "Editor" ]
+			"TargetAllowList": [
+				"Editor"
+			]
 		}
 	]
 }
\ No newline at end of file
diff --git a/Engine/Plugins/Experimental/PCGMeshPartitionInterop/PCGMeshPartitionInterop.uplugin b/Engine/Plugins/Experimental/PCGMeshPartitionInterop/PCGMeshPartitionInterop.uplugin
--- a/Engine/Plugins/Experimental/PCGMeshPartitionInterop/PCGMeshPartitionInterop.uplugin
+++ b/Engine/Plugins/Experimental/PCGMeshPartitionInterop/PCGMeshPartitionInterop.uplugin
@@ -1,10 +1,10 @@
 {
 	"FileVersion": 3,
 	"Version": 1,
-	"VersionName": "0.1",
-	"FriendlyName": "Procedural Content Generation Framework (PCG) MeshPartition Interop",
-	"Description": "Extra plugin for Procedural Content Generation Framework interacting with the Mega Mesh system.",
-	"Category": "Editor",
+	"VersionName": "1.0",
+	"FriendlyName": "Procedural Content Generation Framework (PCG) Mesh Partition Interop",
+	"Description": "Interoperability of Mesh Partition with PCG.",
+	"Category": "Mesh Partition",
 	"CreatedBy": "Epic Games, Inc.",
 	"CreatedByURL": "https://epicgames.com",
 	"DocsURL": "https://docs.unrealengine.com/latest/en-US/procedural-content-generation--framework-in-unreal-engine/",
@@ -41,4 +41,4 @@
 			"Enabled": true
 		}
 	]
-}
+}
\ No newline at end of file
