tweak mesh terrain tool UI so header is not in the scroll block

#Synced-CL 52243346

diff --git a/Engine/Plugins/Experimental/MeshTerrainMode/Source/MeshTerrainMode/Private/MeshTerrainModeToolkit_Details.cpp b/Engine/Plugins/Experimental/MeshTerrainMode/Source/MeshTerrainMode/Private/MeshTerrainModeToolkit_Details.cpp
--- a/Engine/Plugins/Experimental/MeshTerrainMode/Source/MeshTerrainMode/Private/MeshTerrainModeToolkit_Details.cpp
+++ b/Engine/Plugins/Experimental/MeshTerrainMode/Source/MeshTerrainMode/Private/MeshTerrainModeToolkit_Details.cpp
@@ -92,54 +92,59 @@ void FMeshTerrainModeToolkit::MakeFullDetailsOverlayWidget()
 						.WidthOverride(300.f)
 						.MaxDesiredHeight(400.f)
 						[
-							SNew(SScrollBox)
-							.Orientation(Orient_Vertical)
-							.ConsumeMouseWheel(EConsumeMouseWheel::Always)
-							+ SScrollBox::Slot()
+							SNew(SVerticalBox)
+							+ SVerticalBox::Slot()
+							.AutoHeight()
+							.Padding(0.f, 0.f, 0.f, 2.f)
 							[
-								SNew(SVerticalBox)
-								+ SVerticalBox::Slot()
-								.AutoHeight()
-								.Padding(0.f, 0.f, 0.f, 2.f)
+								SNew(SHorizontalBox)
+								+ SHorizontalBox::Slot()
+								.AutoWidth()
+								.Padding(8.f, 0.f, 0.f, 4.f)
+								.VAlign(VAlign_Bottom)
 								[
-									SNew(SHorizontalBox)
-									+ SHorizontalBox::Slot()
-									.AutoWidth()
-									.Padding(8.f, 0.f, 0.f, 4.f)
-									.VAlign(VAlign_Bottom)
-									[
-										SNew(SImage)
-										.Image_Lambda([this]()
-										{
-											return ActiveToolIcon;
-										})
-										.DesiredSizeOverride(FVector2D(16.f))
-										.ColorAndOpacity(FSlateColor::UseForeground())
-									]
-									+ SHorizontalBox::Slot()
-									.FillWidth(1.0f)
-									.Padding(8.f, 0.f, 0.f, 0.f)
-									.VAlign(VAlign_Bottom)
-									.HAlign(HAlign_Left)
-									[
-										SNew(STextBlock)
-										.Font(TitleFont)
-										.Text_Lambda([this]()
-										{
-											return ActiveToolName;
-										})
-									]
+									SNew(SImage)
+									.Image_Lambda([this]()
+									{
+										return ActiveToolIcon;
+									})
+									.DesiredSizeOverride(FVector2D(16.f))
+									.ColorAndOpacity(FSlateColor::UseForeground())
 								]
-								+ SVerticalBox::Slot()
-								.AutoHeight()
+								+ SHorizontalBox::Slot()
+								.FillWidth(1.0f)
 								.Padding(8.f, 0.f, 0.f, 0.f)
+								.VAlign(VAlign_Bottom)
+								.HAlign(HAlign_Left)
 								[
-									ToolWarningArea.ToSharedRef()
+									SNew(STextBlock)
+									.Font(TitleFont)
+									.Text_Lambda([this]()
+									{
+										return ActiveToolName;
+									})
 								]
-								+ SVerticalBox::Slot()
-								.AutoHeight()
+							]
+							+ SVerticalBox::Slot()
+							.FillHeight(1.0f)
+							[
+								SNew(SScrollBox)
+								.Orientation(Orient_Vertical)
+								.ConsumeMouseWheel(EConsumeMouseWheel::Always)
+								+ SScrollBox::Slot()
 								[
-									ModeDetailsView.ToSharedRef()
+									SNew(SVerticalBox)
+									+ SVerticalBox::Slot()
+									.AutoHeight()
+									.Padding(8.f, 0.f, 0.f, 0.f)
+									[
+										ToolWarningArea.ToSharedRef()
+									]
+									+ SVerticalBox::Slot()
+									.AutoHeight()
+									[
+										ModeDetailsView.ToSharedRef()
+									]
 								]
 							]
 						]
