Mesh Terrain: add a flag to accept tool changes when hitting Esc

#Synced-CL 52738760

diff --git a/Engine/Plugins/Experimental/MeshTerrainMode/Source/MeshTerrainMode/Private/MeshTerrainMode.cpp b/Engine/Plugins/Experimental/MeshTerrainMode/Source/MeshTerrainMode/Private/MeshTerrainMode.cpp
--- a/Engine/Plugins/Experimental/MeshTerrainMode/Source/MeshTerrainMode/Private/MeshTerrainMode.cpp
+++ b/Engine/Plugins/Experimental/MeshTerrainMode/Source/MeshTerrainMode/Private/MeshTerrainMode.cpp
@@ -2333,7 +2333,7 @@ void UMeshTerrainMode::BindCommands()
 
 	CommandList->MapAction(
 		ToolManagerCommands.CancelOrCompleteActiveTool,
-		FExecuteAction::CreateLambda([this]() { CancelActiveToolActionOrTool(); }),
+		FExecuteAction::CreateLambda([this]() { OnEscapeKey(); }),
 		FCanExecuteAction::CreateLambda([this]() {
 				return GetInteractiveToolsContext()->CanCompleteActiveTool() || GetInteractiveToolsContext()->CanCancelActiveTool();
 			}),
@@ -2386,6 +2386,18 @@ void UMeshTerrainMode::CancelActiveToolActionOrTool()
 	GetInteractiveToolsContext()->EndTool(ShutdownType);
 }
 
+void UMeshTerrainMode::OnEscapeKey()
+{
+	if (GetDefault<UMeshTerrainModeCustomizationSettings>()->bEscapeAcceptsToolResult)
+	{
+		AcceptActiveToolActionOrTool();
+	}
+	else
+	{
+		CancelActiveToolActionOrTool();
+	}
+}
+
 void UMeshTerrainMode::ModelingModeShortcutRequested(EMeshTerrainModeActionCommands Command)
 {
 	if (Command == EMeshTerrainModeActionCommands::FocusViewToCursor)
diff --git a/Engine/Plugins/Experimental/MeshTerrainMode/Source/MeshTerrainMode/Private/MeshTerrainMode.h b/Engine/Plugins/Experimental/MeshTerrainMode/Source/MeshTerrainMode/Private/MeshTerrainMode.h
--- a/Engine/Plugins/Experimental/MeshTerrainMode/Source/MeshTerrainMode/Private/MeshTerrainMode.h
+++ b/Engine/Plugins/Experimental/MeshTerrainMode/Source/MeshTerrainMode/Private/MeshTerrainMode.h
@@ -168,6 +168,7 @@ protected:
 
 	void AcceptActiveToolActionOrTool();
 	void CancelActiveToolActionOrTool();
+	void OnEscapeKey();
 
 	void ConfigureRealTimeViewportsOverride(bool bEnable);
 
diff --git a/Engine/Plugins/Experimental/MeshTerrainMode/Source/MeshTerrainMode/Private/MeshTerrainModeSettings.h b/Engine/Plugins/Experimental/MeshTerrainMode/Source/MeshTerrainMode/Private/MeshTerrainModeSettings.h
--- a/Engine/Plugins/Experimental/MeshTerrainMode/Source/MeshTerrainMode/Private/MeshTerrainModeSettings.h
+++ b/Engine/Plugins/Experimental/MeshTerrainMode/Source/MeshTerrainMode/Private/MeshTerrainModeSettings.h
@@ -323,6 +323,10 @@ public:
 	UPROPERTY(config, EditAnywhere, Category = "Mesh Terrain Mode|UI Customization")
 	bool bAlwaysShowToolButtons = false;
 
+	/** If true, pressing Escape will Accept the active tool result instead of Cancelling it. Defaults to off. */
+	UPROPERTY(config, EditAnywhere, Category = "Mesh Terrain Mode|UI Customization")
+	bool bEscapeAcceptsToolResult = false;
+
 	/**
 	 * Custom color for when geometry is not a part of current selection when using Mesh Element Selection
 	 */
diff --git a/Engine/Plugins/Experimental/MeshTerrainMode/Source/MeshTerrainMode/Private/MeshTerrainModeToolkit.cpp b/Engine/Plugins/Experimental/MeshTerrainMode/Source/MeshTerrainMode/Private/MeshTerrainModeToolkit.cpp
--- a/Engine/Plugins/Experimental/MeshTerrainMode/Source/MeshTerrainMode/Private/MeshTerrainModeToolkit.cpp
+++ b/Engine/Plugins/Experimental/MeshTerrainMode/Source/MeshTerrainMode/Private/MeshTerrainModeToolkit.cpp
@@ -612,8 +612,12 @@ void FMeshTerrainModeToolkit::MakeToolShutdownOverlayWidget()
 	};
 	auto GetCancelButtonTooltip = [this]()
 	{
-		return AcceptCancelButtonParams.IsSet() && AcceptCancelButtonParams->OverrideCancelButtonTooltip.IsSet() ?
-			AcceptCancelButtonParams->OverrideCancelButtonTooltip.GetValue()
+		if (AcceptCancelButtonParams.IsSet() && AcceptCancelButtonParams->OverrideCancelButtonTooltip.IsSet())
+		{
+			return AcceptCancelButtonParams->OverrideCancelButtonTooltip.GetValue();
+		}
+		return GetDefault<UMeshTerrainModeCustomizationSettings>()->bEscapeAcceptsToolResult
+			? LOCTEXT("OverlayCancelTooltip_NoEsc", "Cancel the active Tool")
 			: LOCTEXT("OverlayCancelTooltip", "Cancel the active Tool [Esc]");
 	};
 	auto GetCancelButtonEnabled = [this]() 
@@ -2150,8 +2154,12 @@ TSharedPtr<SWidget> FMeshTerrainModeToolkit::MakeToolShutdownButtons()
 	// cancel button param funcs
 	TFunction<FText()> GetCancelButtonTooltip = [this]()
 		{
-			return AcceptCancelButtonParams.IsSet() && AcceptCancelButtonParams->OverrideCancelButtonTooltip.IsSet() ?
-				AcceptCancelButtonParams->OverrideCancelButtonTooltip.GetValue()
+			if (AcceptCancelButtonParams.IsSet() && AcceptCancelButtonParams->OverrideCancelButtonTooltip.IsSet())
+			{
+				return AcceptCancelButtonParams->OverrideCancelButtonTooltip.GetValue();
+			}
+			return GetDefault<UMeshTerrainModeCustomizationSettings>()->bEscapeAcceptsToolResult
+				? LOCTEXT("OverlayCancelTooltip_NoEsc", "Cancel the active Tool")
 				: LOCTEXT("OverlayCancelTooltip", "Cancel the active Tool [Esc]");
 		};
 	TFunction<bool()> GetCancelButtonEnabled = [this]() 
